Zen server pluggable transports "SDK" bits, not currently directly used by anything in the engine tree

#skipci

diff --git a/Engine/Source/Developer/ZenPluggableTransport/README.md b/Engine/Source/Developer/ZenPluggableTransport/README.md
new file mode 100644
--- /dev/null
+++ b/Engine/Source/Developer/ZenPluggableTransport/README.md
@@ -0,0 +1,5 @@
+This code corresponds to the code in [UE5/Engine/Source/Developer/ZenPluggableTransport](https://github.com/EpicGames/UnrealEngine/tree/release/Engine/Source/Developer), 
+and provides the API definitions for creating transport plug-ins for use with the
+Zen server. Pluggable transports allow us to support a variety of transport mechanisms, 
+including some which are not possible to share with a general audience. These are 
+developed and maintained in the UE tree or elsewhere for logistical and legal reasons.
diff --git a/Engine/Source/Developer/ZenPluggableTransport/transport-sdk/include/transportplugin.h b/Engine/Source/Developer/ZenPluggableTransport/transport-sdk/include/transportplugin.h
new file mode 100644
--- /dev/null
+++ b/Engine/Source/Developer/ZenPluggableTransport/transport-sdk/include/transportplugin.h
@@ -0,0 +1,111 @@
+// Copyright Epic Games, Inc. All Rights Reserved.
+
+#pragma once
+
+#include <stdint.h>
+
+// Important note: this header is meant to compile standalone
+// and should therefore not depend on anything from the Zen codebase
+
+namespace zen {
+
+class TransportConnection;
+class TransportPlugin;
+class TransportServerConnection;
+class TransportServer;
+
+/*************************************************************************
+
+   The following interfaces are implemented on the server side, and instances
+   are provided to the plugins.
+
+*************************************************************************/
+
+/** Plugin-server interface for connection
+
+	This is returned by a call to TransportServer::CreateConnectionHandler
+	and there should be one instance created per established connection
+
+	The plugin uses this interface to feed data into the server side
+	protocol implementation which will parse the incoming messages and
+	dispatch to appropriate request handlers and ultimately call into
+	TransportConnection functions which write data back to the client
+ */
+class TransportServerConnection
+{
+public:
+	virtual uint32_t AddRef() const									  = 0;
+	virtual uint32_t Release() const								  = 0;
+	virtual void	 OnBytesRead(const void* Buffer, size_t DataSize) = 0;
+};
+
+/** Server interface
+
+ There will be one instance of this provided by the system to the transport plugin
+
+ The plugin can use this to register new connections
+
+ */
+class TransportServer
+{
+public:
+	virtual TransportServerConnection* CreateConnectionHandler(TransportConnection* Connection) = 0;
+};
+
+/*************************************************************************
+
+   The following interfaces are to be implemented by transport plugins.
+
+*************************************************************************/
+
+/** Interface which needs to be implemented by a transport plugin
+
+	This is responsible for setting up and running the communication
+	for a given transport.
+
+	Once initialized, the plugin should be ready to accept connections
+	using its own execution resources (threads, thread pools etc)
+ */
+class TransportPlugin
+{
+public:
+	virtual uint32_t AddRef() const								  = 0;
+	virtual uint32_t Release() const							  = 0;
+	virtual void	 Initialize(TransportServer* ServerInterface) = 0;
+	virtual void	 Shutdown()									  = 0;
+
+	/** Check whether this transport is usable.
+	 */
+	virtual bool IsAvailable() = 0;
+};
+
+/** A transport plugin provider needs to implement this interface
+
+   The plugin should create one instance of this per established
+   connection and register it with the TransportServer instance
+   CreateConnectionHandler() function. The server will subsequently
+   use this interface to write response data back to the client and
+   to manage the connection life cycle in general
+*/
+class TransportConnection
+{
+public:
+	virtual int64_t WriteBytes(const void* Buffer, size_t DataSize) = 0;
+	virtual void	Shutdown(bool Receive, bool Transmit)			= 0;
+	virtual void	CloseConnection()								= 0;
+};
+
+}  // namespace zen
+
+#if defined(_MSC_VER)
+#	define DLL_TRANSPORT_API __declspec(dllexport)
+#else
+#	define DLL_TRANSPORT_API
+#endif
+
+extern "C"
+{
+	DLL_TRANSPORT_API zen::TransportPlugin* CreateTransportPlugin();
+}
+
+typedef zen::TransportPlugin* (*PfnCreateTransportPlugin)();
diff --git a/Engine/Source/Developer/ZenPluggableTransport/winsock/winsock.cpp b/Engine/Source/Developer/ZenPluggableTransport/winsock/winsock.cpp
new file mode 100644
--- /dev/null
+++ b/Engine/Source/Developer/ZenPluggableTransport/winsock/winsock.cpp
@@ -0,0 +1,352 @@
+// Copyright Epic Games, Inc. All Rights Reserved.
+
+#include <inttypes.h>
+#include <atomic>
+#include <exception>
+#include <future>
+#include <memory>
+#include <thread>
+
+#include <zencore/refcount.h>
+#include <zencore/zencore.h>
+
+#ifndef _WIN32_WINNT
+#	define _WIN32_WINNT 0x0A00
+#endif
+
+ZEN_THIRD_PARTY_INCLUDES_START
+#include <WS2tcpip.h>
+#include <WinSock2.h>
+#include <windows.h>
+ZEN_THIRD_PARTY_INCLUDES_END
+
+#include <transportplugin.h>
+
+//////////////////////////////////////////////////////////////////////////
+
+using namespace zen;
+
+class SocketTransportPlugin : public TransportPlugin, zen::RefCounted
+{
+public:
+	SocketTransportPlugin(uint16_t BasePort, unsigned int ThreadCount);
+	~SocketTransportPlugin();
+
+	// TransportPlugin implementation
+
+	virtual uint32_t AddRef() const override;
+	virtual uint32_t Release() const override;
+	virtual void	 Initialize(TransportServer* ServerInterface) override;
+	virtual void	 Shutdown() override;
+	virtual bool	 IsAvailable() override;
+
+private:
+	TransportServer* m_ServerInterface = nullptr;
+	bool			 m_IsOk			   = true;
+	uint16_t		 m_BasePort		   = 0;
+	int				 m_ThreadCount	   = 0;
+
+	SOCKET						   m_ListenSocket{};
+	std::thread					   m_AcceptThread;
+	std::atomic_flag			   m_KeepRunning;
+	std::vector<std::future<void>> m_Connections;
+};
+
+struct SocketTransportConnection : public TransportConnection
+{
+public:
+	SocketTransportConnection();
+	~SocketTransportConnection();
+
+	void Initialize(TransportServerConnection* ServerConnection, SOCKET ClientSocket);
+	void HandleConnection();
+
+	// TransportConnection implementation
+
+	virtual int64_t WriteBytes(const void* Buffer, size_t DataSize) override;
+	virtual void	Shutdown(bool Receive, bool Transmit) override;
+	virtual void	CloseConnection() override;
+
+private:
+	zen::Ref<TransportServerConnection> m_ConnectionHandler;
+	SOCKET								m_ClientSocket{};
+	bool								m_IsTerminated = false;
+};
+
+//////////////////////////////////////////////////////////////////////////
+
+SocketTransportConnection::SocketTransportConnection()
+{
+}
+
+SocketTransportConnection::~SocketTransportConnection()
+{
+}
+
+void
+SocketTransportConnection::Initialize(TransportServerConnection* ServerConnection, SOCKET ClientSocket)
+{
+	// ZEN_ASSERT(!m_ConnectionHandler);
+
+	m_ConnectionHandler = ServerConnection;
+	m_ClientSocket		= ClientSocket;
+}
+
+void
+SocketTransportConnection::HandleConnection()
+{
+	// ZEN_ASSERT(m_ConnectionHandler);
+
+	const int				   InputBufferSize = 64 * 1024;
+	std::unique_ptr<uint8_t[]> InputBuffer{new uint8_t[64 * 1024]};
+
+	do
+	{
+		const int RecvBytes = recv(m_ClientSocket, (char*)InputBuffer.get(), InputBufferSize, /* flags */ 0);
+
+		if (RecvBytes == 0)
+		{
+			// Connection closed
+			return CloseConnection();
+		}
+		else if (RecvBytes < 0)
+		{
+			// Error
+			return CloseConnection();
+		}
+
+		m_ConnectionHandler->OnBytesRead(InputBuffer.get(), RecvBytes);
+	} while (m_ClientSocket);
+}
+
+void
+SocketTransportConnection::CloseConnection()
+{
+	if (m_IsTerminated)
+	{
+		return;
+	}
+
+	// ZEN_ASSERT(m_ClientSocket);
+	m_IsTerminated = true;
+
+	shutdown(m_ClientSocket, SD_BOTH);	// We won't be sending or receiving any more data
+
+	closesocket(m_ClientSocket);
+	m_ClientSocket = 0;
+}
+
+int64_t
+SocketTransportConnection::WriteBytes(const void* Buffer, size_t DataSize)
+{
+	const uint8_t* BufferCursor	  = reinterpret_cast<const uint8_t*>(Buffer);
+	int64_t		   TotalBytesSent = 0;
+
+	while (DataSize)
+	{
+		const int MaxBlockSize	= 128 * 1024;
+		const int SendBlockSize = (DataSize > MaxBlockSize) ? MaxBlockSize : (int)DataSize;
+		const int SentBytes		= send(m_ClientSocket, (const char*)BufferCursor, SendBlockSize, /* flags */ 0);
+
+		if (SentBytes < 0)
+		{
+			// Error
+			return SentBytes;
+		}
+
+		BufferCursor += SentBytes;
+		DataSize -= SentBytes;
+		TotalBytesSent += SentBytes;
+	}
+
+	return TotalBytesSent;
+}
+
+void
+SocketTransportConnection::Shutdown(bool Receive, bool Transmit)
+{
+	if (Receive)
+	{
+		if (Transmit)
+		{
+			shutdown(m_ClientSocket, SD_BOTH);
+		}
+		else
+		{
+			shutdown(m_ClientSocket, SD_RECEIVE);
+		}
+	}
+	else if (Transmit)
+	{
+		shutdown(m_ClientSocket, SD_SEND);
+	}
+}
+
+//////////////////////////////////////////////////////////////////////////
+
+SocketTransportPlugin::SocketTransportPlugin(uint16_t BasePort, unsigned int ThreadCount)
+: m_BasePort(BasePort)
+, m_ThreadCount(ThreadCount != 0 ? ThreadCount : std::max(std::thread::hardware_concurrency(), 8u))
+{
+#if ZEN_PLATFORM_WINDOWS
+	WSADATA wsaData;
+	if (int Result = WSAStartup(0x202, &wsaData); Result != 0)
+	{
+		m_IsOk = false;
+		WSACleanup();
+	}
+#endif
+}
+
+SocketTransportPlugin::~SocketTransportPlugin()
+{
+	Shutdown();
+
+#if ZEN_PLATFORM_WINDOWS
+	if (m_IsOk)
+	{
+		WSACleanup();
+	}
+#endif
+}
+
+uint32_t
+SocketTransportPlugin::AddRef() const
+{
+	return RefCounted::AddRef();
+}
+
+uint32_t
+SocketTransportPlugin::Release() const
+{
+	return RefCounted::Release();
+}
+
+void
+SocketTransportPlugin::Initialize(TransportServer* ServerInterface)
+{
+	uint16_t Port = m_BasePort;
+
+	m_ServerInterface = ServerInterface;
+	m_ListenSocket	  = socket(AF_INET6, SOCK_STREAM, 0);
+
+	if (m_ListenSocket == SOCKET_ERROR || m_ListenSocket == INVALID_SOCKET)
+	{
+		throw std::system_error(std::error_code(WSAGetLastError(), std::system_category()),
+								"socket creation failed in HTTP plugin server init");
+	}
+
+	sockaddr_in6 Server{};
+	Server.sin6_family = AF_INET6;
+	Server.sin6_port   = htons(Port);
+	Server.sin6_addr   = in6addr_any;
+
+	if (int Result = bind(m_ListenSocket, (sockaddr*)&Server, sizeof(Server)); Result == SOCKET_ERROR)
+	{
+		throw std::system_error(std::error_code(WSAGetLastError(), std::system_category()), "bind call failed in HTTP plugin server init");
+	}
+
+	if (int Result = listen(m_ListenSocket, AF_INET6); Result == SOCKET_ERROR)
+	{
+		throw std::system_error(std::error_code(WSAGetLastError(), std::system_category()),
+								"listen call failed in HTTP plugin server init");
+	}
+
+	m_KeepRunning.test_and_set();
+
+	m_AcceptThread = std::thread([&] {
+		// SetCurrentThreadName("http_plugin_acceptor");
+
+		// ZEN_INFO("HTTP plugin server waiting for connections");
+
+		do
+		{
+			if (SOCKET ClientSocket = accept(m_ListenSocket, NULL, NULL); ClientSocket != SOCKET_ERROR)
+			{
+				int Flag = 1;
+				setsockopt(ClientSocket, IPPROTO_TCP, TCP_NODELAY, (char*)&Flag, sizeof(Flag));
+
+				// Handle new connection
+				SocketTransportConnection* Connection = new SocketTransportConnection();
+				TransportServerConnection* ConnectionInterface{m_ServerInterface->CreateConnectionHandler(Connection)};
+				Connection->Initialize(ConnectionInterface, ClientSocket);
+
+				m_Connections.push_back(std::async(std::launch::async, [Connection] {
+					try
+					{
+						Connection->HandleConnection();
+					}
+					catch (std::exception&)
+					{
+						// ZEN_WARN("exception caught in connection loop: {}", Ex.what());
+					}
+
+					delete Connection;
+				}));
+			}
+			else
+			{
+			}
+		} while (m_KeepRunning.test());
+
+		// ZEN_INFO("HTTP plugin server accept thread exit");
+	});
+}
+
+void
+SocketTransportPlugin::Shutdown()
+{
+	// TODO: all pending/ongoing work should be drained here as well
+
+	m_KeepRunning.clear();
+
+	closesocket(m_ListenSocket);
+	m_ListenSocket = 0;
+
+	if (m_AcceptThread.joinable())
+	{
+		m_AcceptThread.join();
+	}
+}
+
+bool
+SocketTransportPlugin::IsAvailable()
+{
+	return true;
+}
+
+//////////////////////////////////////////////////////////////////////////
+
+TransportPlugin*
+CreateTransportPlugin()
+{
+	return new SocketTransportPlugin(1337, 8);
+}
+
+BOOL WINAPI
+DllMain([[maybe_unused]] HINSTANCE hinstDLL,   // handle to DLL module
+		DWORD					   fdwReason,  // reason for calling function
+		LPVOID					   lpvReserved)					   // reserved
+{
+	// Perform actions based on the reason for calling.
+	switch (fdwReason)
+	{
+		case DLL_PROCESS_ATTACH:
+			break;
+
+		case DLL_THREAD_ATTACH:
+			break;
+
+		case DLL_THREAD_DETACH:
+			break;
+
+		case DLL_PROCESS_DETACH:
+			if (lpvReserved != nullptr)
+			{
+				break;	// do not do cleanup if process termination scenario
+			}
+			break;
+	}
+
+	return TRUE;
+}
