update transports SDK

#skipci

diff --git a/Engine/Source/Developer/ZenPluggableTransport/transport-sdk/include/transportplugin.h b/Engine/Source/Developer/ZenPluggableTransport/transport-sdk/include/transportplugin.h
--- a/Engine/Source/Developer/ZenPluggableTransport/transport-sdk/include/transportplugin.h
+++ b/Engine/Source/Developer/ZenPluggableTransport/transport-sdk/include/transportplugin.h
@@ -4,8 +4,16 @@
 
 #include <stdint.h>
 
-// Important note: this header is meant to compile standalone
-// and should therefore not depend on anything from the Zen codebase
+// This header is meant to compile standalone and should therefore NOT depend
+// on anything from the Zen tree
+
+//////////////////////////////////////////////////////////////////////////
+//
+// IMPORTANT: Any development or changes to this header should be made
+// in the Zen repository http://github.com/epicgames/zen even if you
+// may find the header in the UE tree
+//
+//////////////////////////////////////////////////////////////////////////
 
 namespace zen {
 
@@ -69,10 +77,11 @@ public:
 class TransportPlugin
 {
 public:
-	virtual uint32_t AddRef() const								  = 0;
-	virtual uint32_t Release() const							  = 0;
-	virtual void	 Initialize(TransportServer* ServerInterface) = 0;
-	virtual void	 Shutdown()									  = 0;
+	virtual uint32_t AddRef() const											   = 0;
+	virtual uint32_t Release() const										   = 0;
+	virtual void	 Configure(const char* OptionTag, const char* OptionValue) = 0;
+	virtual void	 Initialize(TransportServer* ServerInterface)			   = 0;
+	virtual void	 Shutdown()												   = 0;
 
 	/** Check whether this transport is usable.
 	 */
diff --git a/Engine/Source/Developer/ZenPluggableTransport/winsock/winsock.cpp b/Engine/Source/Developer/ZenPluggableTransport/winsock/winsock.cpp
--- a/Engine/Source/Developer/ZenPluggableTransport/winsock/winsock.cpp
+++ b/Engine/Source/Developer/ZenPluggableTransport/winsock/winsock.cpp
@@ -2,9 +2,11 @@
 
 #include <inttypes.h>
 #include <atomic>
+#include <charconv>
 #include <exception>
 #include <future>
 #include <memory>
+#include <optional>
 #include <thread>
 
 #include <zencore/refcount.h>
@@ -24,18 +26,34 @@ ZEN_THIRD_PARTY_INCLUDES_END
 
 //////////////////////////////////////////////////////////////////////////
 
+template<Integral T>
+std::optional<T>
+ParseInt(const std::string_view& Input)
+{
+	T							 Out	= 0;
+	const std::from_chars_result Result = std::from_chars(Input.data(), Input.data() + Input.size(), Out);
+	if (Result.ec == std::errc::invalid_argument || Result.ec == std::errc::result_out_of_range)
+	{
+		return std::nullopt;
+	}
+	return Out;
+}
+
+//////////////////////////////////////////////////////////////////////////
+
 using namespace zen;
 
-class SocketTransportPlugin : public TransportPlugin, zen::RefCounted
+class WinsockTransportPlugin : public TransportPlugin, zen::RefCounted
 {
 public:
-	SocketTransportPlugin(uint16_t BasePort, unsigned int ThreadCount);
-	~SocketTransportPlugin();
+	WinsockTransportPlugin();
+	~WinsockTransportPlugin();
 
 	// TransportPlugin implementation
 
 	virtual uint32_t AddRef() const override;
 	virtual uint32_t Release() const override;
+	virtual void	 Configure(const char* OptionTag, const char* OptionValue) override;
 	virtual void	 Initialize(TransportServer* ServerInterface) override;
 	virtual void	 Shutdown() override;
 	virtual bool	 IsAvailable() override;
@@ -44,7 +62,6 @@ private:
 	TransportServer* m_ServerInterface = nullptr;
 	bool			 m_IsOk			   = true;
 	uint16_t		 m_BasePort		   = 0;
-	int				 m_ThreadCount	   = 0;
 
 	SOCKET						   m_ListenSocket{};
 	std::thread					   m_AcceptThread;
@@ -52,11 +69,11 @@ private:
 	std::vector<std::future<void>> m_Connections;
 };
 
-struct SocketTransportConnection : public TransportConnection
+struct WinsockTransportConnection : public TransportConnection
 {
 public:
-	SocketTransportConnection();
-	~SocketTransportConnection();
+	WinsockTransportConnection();
+	~WinsockTransportConnection();
 
 	void Initialize(TransportServerConnection* ServerConnection, SOCKET ClientSocket);
 	void HandleConnection();
@@ -75,16 +92,16 @@ private:
 
 //////////////////////////////////////////////////////////////////////////
 
-SocketTransportConnection::SocketTransportConnection()
+WinsockTransportConnection::WinsockTransportConnection()
 {
 }
 
-SocketTransportConnection::~SocketTransportConnection()
+WinsockTransportConnection::~WinsockTransportConnection()
 {
 }
 
 void
-SocketTransportConnection::Initialize(TransportServerConnection* ServerConnection, SOCKET ClientSocket)
+WinsockTransportConnection::Initialize(TransportServerConnection* ServerConnection, SOCKET ClientSocket)
 {
 	// ZEN_ASSERT(!m_ConnectionHandler);
 
@@ -93,7 +110,7 @@ SocketTransportConnection::Initialize(TransportServerConnection* ServerConnectio
 }
 
 void
-SocketTransportConnection::HandleConnection()
+WinsockTransportConnection::HandleConnection()
 {
 	// ZEN_ASSERT(m_ConnectionHandler);
 
@@ -120,7 +137,7 @@ SocketTransportConnection::HandleConnection()
 }
 
 void
-SocketTransportConnection::CloseConnection()
+WinsockTransportConnection::CloseConnection()
 {
 	if (m_IsTerminated)
 	{
@@ -137,7 +154,7 @@ SocketTransportConnection::CloseConnection()
 }
 
 int64_t
-SocketTransportConnection::WriteBytes(const void* Buffer, size_t DataSize)
+WinsockTransportConnection::WriteBytes(const void* Buffer, size_t DataSize)
 {
 	const uint8_t* BufferCursor	  = reinterpret_cast<const uint8_t*>(Buffer);
 	int64_t		   TotalBytesSent = 0;
@@ -163,7 +180,7 @@ SocketTransportConnection::WriteBytes(const void* Buffer, size_t DataSize)
 }
 
 void
-SocketTransportConnection::Shutdown(bool Receive, bool Transmit)
+WinsockTransportConnection::Shutdown(bool Receive, bool Transmit)
 {
 	if (Receive)
 	{
@@ -184,9 +201,7 @@ SocketTransportConnection::Shutdown(bool Receive, bool Transmit)
 
 //////////////////////////////////////////////////////////////////////////
 
-SocketTransportPlugin::SocketTransportPlugin(uint16_t BasePort, unsigned int ThreadCount)
-: m_BasePort(BasePort)
-, m_ThreadCount(ThreadCount != 0 ? ThreadCount : std::max(std::thread::hardware_concurrency(), 8u))
+WinsockTransportPlugin::WinsockTransportPlugin()
 {
 #if ZEN_PLATFORM_WINDOWS
 	WSADATA wsaData;
@@ -198,7 +213,7 @@ SocketTransportPlugin::SocketTransportPlugin(uint16_t BasePort, unsigned int Thr
 #endif
 }
 
-SocketTransportPlugin::~SocketTransportPlugin()
+WinsockTransportPlugin::~WinsockTransportPlugin()
 {
 	Shutdown();
 
@@ -211,19 +226,37 @@ SocketTransportPlugin::~SocketTransportPlugin()
 }
 
 uint32_t
-SocketTransportPlugin::AddRef() const
+WinsockTransportPlugin::AddRef() const
 {
 	return RefCounted::AddRef();
 }
 
 uint32_t
-SocketTransportPlugin::Release() const
+WinsockTransportPlugin::Release() const
 {
 	return RefCounted::Release();
 }
 
 void
-SocketTransportPlugin::Initialize(TransportServer* ServerInterface)
+WinsockTransportPlugin::Configure(const char* OptionTag, const char* OptionValue)
+{
+	using namespace std::literals;
+
+	if (OptionTag == "port"sv)
+	{
+		if (auto PortNum = ParseInt<uint16_t>(OptionValue))
+		{
+			m_BasePort = *PortNum;
+		}
+	}
+	else
+	{
+		// Unknown configuration option
+	}
+}
+
+void
+WinsockTransportPlugin::Initialize(TransportServer* ServerInterface)
 {
 	uint16_t Port = m_BasePort;
 
@@ -267,8 +300,8 @@ SocketTransportPlugin::Initialize(TransportServer* ServerInterface)
 				setsockopt(ClientSocket, IPPROTO_TCP, TCP_NODELAY, (char*)&Flag, sizeof(Flag));
 
 				// Handle new connection
-				SocketTransportConnection* Connection = new SocketTransportConnection();
-				TransportServerConnection* ConnectionInterface{m_ServerInterface->CreateConnectionHandler(Connection)};
+				WinsockTransportConnection* Connection = new WinsockTransportConnection();
+				TransportServerConnection*	ConnectionInterface{m_ServerInterface->CreateConnectionHandler(Connection)};
 				Connection->Initialize(ConnectionInterface, ClientSocket);
 
 				m_Connections.push_back(std::async(std::launch::async, [Connection] {
@@ -294,7 +327,7 @@ SocketTransportPlugin::Initialize(TransportServer* ServerInterface)
 }
 
 void
-SocketTransportPlugin::Shutdown()
+WinsockTransportPlugin::Shutdown()
 {
 	// TODO: all pending/ongoing work should be drained here as well
 
@@ -310,7 +343,7 @@ SocketTransportPlugin::Shutdown()
 }
 
 bool
-SocketTransportPlugin::IsAvailable()
+WinsockTransportPlugin::IsAvailable()
 {
 	return true;
 }
@@ -320,7 +353,7 @@ SocketTransportPlugin::IsAvailable()
 TransportPlugin*
 CreateTransportPlugin()
 {
-	return new SocketTransportPlugin(1337, 8);
+	return new WinsockTransportPlugin;
 }
 
 BOOL WINAPI
