new zen transport bits

diff --git a/Engine/Source/Developer/ZenPluggableTransport/transport-sdk/include/transportplugin.h b/Engine/Source/Developer/ZenPluggableTransport/transport-sdk/include/transportplugin.h
--- a/Engine/Source/Developer/ZenPluggableTransport/transport-sdk/include/transportplugin.h
+++ b/Engine/Source/Developer/ZenPluggableTransport/transport-sdk/include/transportplugin.h
@@ -77,11 +77,12 @@ public:
 class TransportPlugin
 {
 public:
-	virtual uint32_t AddRef() const											   = 0;
-	virtual uint32_t Release() const										   = 0;
-	virtual void	 Configure(const char* OptionTag, const char* OptionValue) = 0;
-	virtual void	 Initialize(TransportServer* ServerInterface)			   = 0;
-	virtual void	 Shutdown()												   = 0;
+	virtual uint32_t	AddRef() const											  = 0;
+	virtual uint32_t	Release() const											  = 0;
+	virtual void		Configure(const char* OptionTag, const char* OptionValue) = 0;
+	virtual void		Initialize(TransportServer* ServerInterface)			  = 0;
+	virtual void		Shutdown()												  = 0;
+	virtual const char* GetDebugName()											  = 0;
 
 	/** Check whether this transport is usable.
 	 */
@@ -99,9 +100,10 @@ public:
 class TransportConnection
 {
 public:
-	virtual int64_t WriteBytes(const void* Buffer, size_t DataSize) = 0;
-	virtual void	Shutdown(bool Receive, bool Transmit)			= 0;
-	virtual void	CloseConnection()								= 0;
+	virtual int64_t		WriteBytes(const void* Buffer, size_t DataSize) = 0;
+	virtual void		Shutdown(bool Receive, bool Transmit)			= 0;
+	virtual void		CloseConnection()								= 0;
+	virtual const char* GetDebugName()									= 0;
 };
 
 }  // namespace zen
diff --git a/Engine/Source/Developer/ZenPluggableTransport/winsock/winsock.cpp b/Engine/Source/Developer/ZenPluggableTransport/winsock/winsock.cpp
--- a/Engine/Source/Developer/ZenPluggableTransport/winsock/winsock.cpp
+++ b/Engine/Source/Developer/ZenPluggableTransport/winsock/winsock.cpp
@@ -9,8 +9,8 @@
 #include <optional>
 #include <thread>
 
-#include <zencore/refcount.h>
-#include <zencore/zencore.h>
+#include <zenbase/concepts.h>
+#include <zenbase/refcount.h>
 
 #ifndef _WIN32_WINNT
 #	define _WIN32_WINNT 0x0A00
@@ -51,12 +51,13 @@ public:
 
 	// TransportPlugin implementation
 
-	virtual uint32_t AddRef() const override;
-	virtual uint32_t Release() const override;
-	virtual void	 Configure(const char* OptionTag, const char* OptionValue) override;
-	virtual void	 Initialize(TransportServer* ServerInterface) override;
-	virtual void	 Shutdown() override;
-	virtual bool	 IsAvailable() override;
+	virtual uint32_t	AddRef() const override;
+	virtual uint32_t	Release() const override;
+	virtual void		Configure(const char* OptionTag, const char* OptionValue) override;
+	virtual void		Initialize(TransportServer* ServerInterface) override;
+	virtual void		Shutdown() override;
+	virtual const char* GetDebugName() override;
+	virtual bool		IsAvailable() override;
 
 private:
 	TransportServer* m_ServerInterface = nullptr;
@@ -80,9 +81,10 @@ public:
 
 	// TransportConnection implementation
 
-	virtual int64_t WriteBytes(const void* Buffer, size_t DataSize) override;
-	virtual void	Shutdown(bool Receive, bool Transmit) override;
-	virtual void	CloseConnection() override;
+	virtual int64_t		WriteBytes(const void* Buffer, size_t DataSize) override;
+	virtual void		Shutdown(bool Receive, bool Transmit) override;
+	virtual void		CloseConnection() override;
+	virtual const char* GetDebugName() override;
 
 private:
 	zen::Ref<TransportServerConnection> m_ConnectionHandler;
@@ -153,6 +155,12 @@ WinsockTransportConnection::CloseConnection()
 	m_ClientSocket = 0;
 }
 
+const char*
+WinsockTransportConnection::GetDebugName()
+{
+	return nullptr;
+}
+
 int64_t
 WinsockTransportConnection::WriteBytes(const void* Buffer, size_t DataSize)
 {
@@ -342,6 +350,12 @@ WinsockTransportPlugin::Shutdown()
 	}
 }
 
+const char*
+WinsockTransportPlugin::GetDebugName()
+{
+	return nullptr;
+}
+
 bool
 WinsockTransportPlugin::IsAvailable()
 {
diff --git a/Engine/Source/Developer/ZenPluggableTransport/zenbase/README.md b/Engine/Source/Developer/ZenPluggableTransport/zenbase/README.md
new file mode 100644
--- /dev/null
+++ b/Engine/Source/Developer/ZenPluggableTransport/zenbase/README.md
@@ -0,0 +1,6 @@
+This is a header-only library which mostly contains definitions to
+support different platforms and compilers.
+
+It is part of the zen codebase but is intended to be consumable
+separately to zenbase etc to support standalone transport plug-ins 
+and similar.
diff --git a/Engine/Source/Developer/ZenPluggableTransport/zenbase/include/zenbase/atomic.h b/Engine/Source/Developer/ZenPluggableTransport/zenbase/include/zenbase/atomic.h
new file mode 100644
--- /dev/null
+++ b/Engine/Source/Developer/ZenPluggableTransport/zenbase/include/zenbase/atomic.h
@@ -0,0 +1,74 @@
+// Copyright Epic Games, Inc. All Rights Reserved.
+
+#pragma once
+
+#include <zenbase/zenbase.h>
+
+#if ZEN_COMPILER_MSC
+#	include <intrin.h>
+#else
+#	include <atomic>
+#endif
+
+#include <cinttypes>
+
+namespace zen {
+
+inline uint32_t
+AtomicIncrement(volatile uint32_t& value)
+{
+#if ZEN_COMPILER_MSC
+	return _InterlockedIncrement((long volatile*)&value);
+#else
+	return ((std::atomic<uint32_t>*)(&value))->fetch_add(1, std::memory_order_seq_cst) + 1;
+#endif
+}
+inline uint32_t
+AtomicDecrement(volatile uint32_t& value)
+{
+#if ZEN_COMPILER_MSC
+	return _InterlockedDecrement((long volatile*)&value);
+#else
+	return ((std::atomic<uint32_t>*)(&value))->fetch_sub(1, std::memory_order_seq_cst) - 1;
+#endif
+}
+
+inline uint64_t
+AtomicIncrement(volatile uint64_t& value)
+{
+#if ZEN_COMPILER_MSC
+	return _InterlockedIncrement64((__int64 volatile*)&value);
+#else
+	return ((std::atomic<uint64_t>*)(&value))->fetch_add(1, std::memory_order_seq_cst) + 1;
+#endif
+}
+inline uint64_t
+AtomicDecrement(volatile uint64_t& value)
+{
+#if ZEN_COMPILER_MSC
+	return _InterlockedDecrement64((__int64 volatile*)&value);
+#else
+	return ((std::atomic<uint64_t>*)(&value))->fetch_sub(1, std::memory_order_seq_cst) - 1;
+#endif
+}
+
+inline uint32_t
+AtomicAdd(volatile uint32_t& value, uint32_t amount)
+{
+#if ZEN_COMPILER_MSC
+	return _InterlockedExchangeAdd((long volatile*)&value, amount);
+#else
+	return ((std::atomic<uint32_t>*)(&value))->fetch_add(amount, std::memory_order_seq_cst);
+#endif
+}
+inline uint64_t
+AtomicAdd(volatile uint64_t& value, uint64_t amount)
+{
+#if ZEN_COMPILER_MSC
+	return _InterlockedExchangeAdd64((__int64 volatile*)&value, amount);
+#else
+	return ((std::atomic<uint64_t>*)(&value))->fetch_add(amount, std::memory_order_seq_cst);
+#endif
+}
+
+}  // namespace zen
diff --git a/Engine/Source/Developer/ZenPluggableTransport/zenbase/include/zenbase/concepts.h b/Engine/Source/Developer/ZenPluggableTransport/zenbase/include/zenbase/concepts.h
new file mode 100644
--- /dev/null
+++ b/Engine/Source/Developer/ZenPluggableTransport/zenbase/include/zenbase/concepts.h
@@ -0,0 +1,46 @@
+// Copyright Epic Games, Inc. All Rights Reserved.
+
+#pragma once
+
+#include <zenbase/zenbase.h>
+
+// At the time of writing only ver >= 13 of LLVM's libc++ has an implementation
+// of std::integral. Some platforms like Ubuntu and Mac OS are still on 12.
+#if defined(__cpp_lib_concepts)
+#	include <concepts>
+template<class T>
+concept Integral = std::integral<T>;
+template<class T>
+concept SignedIntegral = std::signed_integral<T>;
+template<class T>
+concept UnsignedIntegral = std::unsigned_integral<T>;
+template<class F, class... A>
+concept Invocable = std::invocable<F, A...>;
+template<class D, class B>
+concept DerivedFrom = std::derived_from<D, B>;
+#else
+#	include <functional>  // for std::invoke below
+
+template<class T>
+concept Integral = std::is_integral_v<T>;
+template<class T>
+concept SignedIntegral = Integral<T> && std::is_signed_v<T>;
+template<class T>
+concept UnsignedIntegral = Integral<T> && !std::is_signed_v<T>;
+template<class F, class... A>
+concept Invocable = requires(F&& f, A&&... a)
+{
+	std::invoke(std::forward<F>(f), std::forward<A>(a)...);
+};
+template<class D, class B>
+concept DerivedFrom = std::is_base_of_v<B, D> && std::is_convertible_v<const volatile D*, const volatile B*>;
+#endif
+
+#if defined(__cpp_lib_ranges)
+#	include <ranges>
+template<typename T>
+concept ContiguousRange = std::ranges::contiguous_range<T>;
+#else
+template<typename T>
+concept ContiguousRange = true;
+#endif
diff --git a/Engine/Source/Developer/ZenPluggableTransport/zenbase/include/zenbase/refcount.h b/Engine/Source/Developer/ZenPluggableTransport/zenbase/include/zenbase/refcount.h
new file mode 100644
--- /dev/null
+++ b/Engine/Source/Developer/ZenPluggableTransport/zenbase/include/zenbase/refcount.h
@@ -0,0 +1,180 @@
+// Copyright Epic Games, Inc. All Rights Reserved.
+#pragma once
+
+#include <zenbase/atomic.h>
+#include <zenbase/concepts.h>
+
+#include <compare>
+
+namespace zen {
+
+/**
+ * Helper base class for reference counted objects using intrusive reference counting
+ */
+class RefCounted
+{
+public:
+	RefCounted()		  = default;
+	virtual ~RefCounted() = default;
+
+	inline uint32_t AddRef() const { return AtomicIncrement(const_cast<RefCounted*>(this)->m_RefCount); }
+	inline uint32_t Release() const
+	{
+		const uint32_t RefCount = AtomicDecrement(const_cast<RefCounted*>(this)->m_RefCount);
+		if (RefCount == 0)
+		{
+			delete this;
+		}
+		return RefCount;
+	}
+
+	// Copying reference counted objects doesn't make a lot of sense generally, so let's prevent it
+
+	RefCounted(const RefCounted&) = delete;
+	RefCounted(RefCounted&&)	  = delete;
+	RefCounted& operator=(const RefCounted&) = delete;
+	RefCounted& operator=(RefCounted&&) = delete;
+
+protected:
+	inline uint32_t RefCount() const { return m_RefCount; }
+
+private:
+	uint32_t m_RefCount = 0;
+};
+
+/**
+ * Smart pointer for classes derived from RefCounted
+ */
+
+template<class T>
+class RefPtr
+{
+public:
+	inline RefPtr() = default;
+	inline RefPtr(const RefPtr& Rhs) : m_Ref(Rhs.m_Ref) { m_Ref && m_Ref->AddRef(); }
+	inline RefPtr(T* Ptr) : m_Ref(Ptr) { m_Ref && m_Ref->AddRef(); }
+	inline ~RefPtr() { m_Ref && m_Ref->Release(); }
+
+	[[nodiscard]] inline bool IsNull() const { return m_Ref == nullptr; }
+	inline explicit			  operator bool() const { return m_Ref != nullptr; }
+	inline					  operator T*() const { return m_Ref; }
+	inline T*				  operator->() const { return m_Ref; }
+
+	inline std::strong_ordering operator<=>(const RefPtr& Rhs) const = default;
+
+	inline RefPtr& operator=(T* Rhs)
+	{
+		Rhs && Rhs->AddRef();
+		m_Ref && m_Ref->Release();
+		m_Ref = Rhs;
+		return *this;
+	}
+	inline RefPtr& operator=(const RefPtr& Rhs)
+	{
+		if (&Rhs != this)
+		{
+			Rhs && Rhs->AddRef();
+			m_Ref && m_Ref->Release();
+			m_Ref = Rhs.m_Ref;
+		}
+		return *this;
+	}
+	inline RefPtr& operator=(RefPtr&& Rhs) noexcept
+	{
+		if (&Rhs != this)
+		{
+			m_Ref && m_Ref->Release();
+			m_Ref	  = Rhs.m_Ref;
+			Rhs.m_Ref = nullptr;
+		}
+		return *this;
+	}
+	template<typename OtherType>
+	inline RefPtr& operator=(RefPtr<OtherType>&& Rhs) noexcept
+	{
+		if ((RefPtr*)&Rhs != this)
+		{
+			m_Ref && m_Ref->Release();
+			m_Ref	  = Rhs.m_Ref;
+			Rhs.m_Ref = nullptr;
+		}
+		return *this;
+	}
+	inline RefPtr(RefPtr&& Rhs) noexcept : m_Ref(Rhs.m_Ref) { Rhs.m_Ref = nullptr; }
+	template<typename OtherType>
+	explicit inline RefPtr(RefPtr<OtherType>&& Rhs) noexcept : m_Ref(Rhs.m_Ref)
+	{
+		Rhs.m_Ref = nullptr;
+	}
+
+private:
+	T* m_Ref = nullptr;
+	template<typename U>
+	friend class RefPtr;
+};
+
+/**
+ * Smart pointer for classes derived from RefCounted
+ *
+ * This variant does not decay to a raw pointer
+ *
+ */
+
+template<class T>
+class Ref
+{
+public:
+	inline Ref() = default;
+	inline Ref(const Ref& Rhs) : m_Ref(Rhs.m_Ref) { m_Ref && m_Ref->AddRef(); }
+	inline explicit Ref(T* Ptr) : m_Ref(Ptr) { m_Ref && m_Ref->AddRef(); }
+	inline ~Ref() { m_Ref && m_Ref->Release(); }
+
+	template<typename DerivedType>
+	requires DerivedFrom<DerivedType, T>
+	inline Ref(const Ref<DerivedType>& Rhs) : Ref(Rhs.m_Ref) {}
+
+	[[nodiscard]] inline bool IsNull() const { return m_Ref == nullptr; }
+	inline explicit			  operator bool() const { return m_Ref != nullptr; }
+	inline T*				  operator->() const { return m_Ref; }
+	inline T&				  operator*() const { return *m_Ref; }
+	inline T*				  Get() const { return m_Ref; }
+
+	inline std::strong_ordering operator<=>(const Ref& Rhs) const = default;
+
+	inline Ref& operator=(T* Rhs)
+	{
+		Rhs && Rhs->AddRef();
+		m_Ref && m_Ref->Release();
+		m_Ref = Rhs;
+		return *this;
+	}
+	inline Ref& operator=(const Ref& Rhs)
+	{
+		if (&Rhs != this)
+		{
+			Rhs && Rhs->AddRef();
+			m_Ref && m_Ref->Release();
+			m_Ref = Rhs.m_Ref;
+		}
+		return *this;
+	}
+	inline Ref& operator=(Ref&& Rhs) noexcept
+	{
+		if (&Rhs != this)
+		{
+			m_Ref && m_Ref->Release();
+			m_Ref	  = Rhs.m_Ref;
+			Rhs.m_Ref = nullptr;
+		}
+		return *this;
+	}
+	inline Ref(Ref&& Rhs) noexcept : m_Ref(Rhs.m_Ref) { Rhs.m_Ref = nullptr; }
+
+private:
+	T* m_Ref = nullptr;
+
+	template<class U>
+	friend class Ref;
+};
+
+}  // namespace zen
diff --git a/Engine/Source/Developer/ZenPluggableTransport/zenbase/include/zenbase/zenbase.h b/Engine/Source/Developer/ZenPluggableTransport/zenbase/include/zenbase/zenbase.h
new file mode 100644
--- /dev/null
+++ b/Engine/Source/Developer/ZenPluggableTransport/zenbase/include/zenbase/zenbase.h
@@ -0,0 +1,250 @@
+// Copyright Epic Games, Inc. All Rights Reserved.
+
+#pragma once
+
+#include <cinttypes>
+#include <version>
+
+#ifdef __clang__
+#	include <type_traits>
+#endif
+
+//////////////////////////////////////////////////////////////////////////
+// Platform
+//
+
+#define ZEN_PLATFORM_WINDOWS 0
+#define ZEN_PLATFORM_LINUX	 0
+#define ZEN_PLATFORM_MAC	 0
+
+#ifdef _WIN32
+#	undef ZEN_PLATFORM_WINDOWS
+#	define ZEN_PLATFORM_WINDOWS 1
+#	define ZEN_PLATFORM_NAME	 "Windows"
+#elif defined(__linux__)
+#	undef ZEN_PLATFORM_LINUX
+#	define ZEN_PLATFORM_LINUX 1
+#	define ZEN_PLATFORM_NAME  "Linux"
+#elif defined(__APPLE__)
+#	undef ZEN_PLATFORM_MAC
+#	define ZEN_PLATFORM_MAC  1
+#	define ZEN_PLATFORM_NAME "MacOS"
+#endif
+
+#if ZEN_PLATFORM_WINDOWS
+#	if !defined(NOMINMAX)
+#		define NOMINMAX  // stops Windows.h from defining 'min/max' macros
+#	endif
+#	if !defined(NOGDI)
+#		define NOGDI
+#	endif
+#	if !defined(WIN32_LEAN_AND_MEAN)
+#		define WIN32_LEAN_AND_MEAN	 // cut-down what Windows.h defines
+#	endif
+#endif
+
+//////////////////////////////////////////////////////////////////////////
+// Compiler
+//
+
+#define ZEN_COMPILER_CLANG 0
+#define ZEN_COMPILER_MSC   0
+#define ZEN_COMPILER_GCC   0
+
+// Clang can define __GNUC__ and/or _MSC_VER so we check for Clang first
+#ifdef __clang__
+#	undef ZEN_COMPILER_CLANG
+#	define ZEN_COMPILER_CLANG 1
+#elif defined(_MSC_VER)
+#	undef ZEN_COMPILER_MSC
+#	define ZEN_COMPILER_MSC 1
+#elif defined(__GNUC__)
+#	undef ZEN_COMPILER_GCC
+#	define ZEN_COMPILER_GCC 1
+#else
+#	error Unknown compiler
+#endif
+
+#if ZEN_COMPILER_MSC
+#	pragma warning(disable : 4324)	 // warning C4324: '<type>': structure was padded due to alignment specifier
+#	pragma warning(default : 4668)	 // warning C4668: 'symbol' is not defined as a preprocessor macro, replacing with '0' for 'directives'
+#	pragma warning(default : 4100)	 // warning C4100: 'identifier' : unreferenced formal parameter
+#	pragma warning( \
+		disable : 4373)	 // '%$S': virtual function overrides '%$pS', previous versions of the compiler did not override when parameters
+						 // only differed by const/volatile qualifiers
+						 // https://docs.microsoft.com/en-us/cpp/error-messages/compiler-warnings/compiler-warning-level-3-c4373
+#endif
+
+#ifndef ZEN_THIRD_PARTY_INCLUDES_START
+#	if ZEN_COMPILER_MSC
+#		define ZEN_THIRD_PARTY_INCLUDES_START                                                                                      \
+			__pragma(warning(push)) __pragma(warning(disable : 4668)) /* C4668: use of undefined preprocessor macro */              \
+				__pragma(warning(disable : 4305))					  /* C4305: 'if': truncation from 'uint32' to 'bool' */         \
+				__pragma(warning(disable : 4267))					  /* C4267: '=': conversion from 'size_t' to 'US' */            \
+				__pragma(warning(disable : 4127))					  /* C4127: conditional expression is constant */               \
+				__pragma(warning(disable : 4189))					  /* C4189: local variable is initialized but not referenced */ \
+				__pragma(warning(disable : 5105)) /* C5105: macro expansion producing 'defined' has undefined behavior */
+#	elif ZEN_COMPILER_CLANG
+#		define ZEN_THIRD_PARTY_INCLUDES_START                                               \
+			_Pragma("clang diagnostic push") _Pragma("clang diagnostic ignored \"-Wundef\"") \
+				_Pragma("clang diagnostic ignored \"-Wunused-parameter\"") _Pragma("clang diagnostic ignored \"-Wunused-variable\"")
+#	elif ZEN_COMPILER_GCC
+#		define ZEN_THIRD_PARTY_INCLUDES_START                                              \
+			_Pragma("GCC diagnostic push") /* NB. ignoring -Wundef doesn't work with GCC */ \
+				_Pragma("GCC diagnostic ignored \"-Wunused-parameter\"") _Pragma("GCC diagnostic ignored \"-Wunused-variable\"")
+#	endif
+#endif
+
+#ifndef ZEN_THIRD_PARTY_INCLUDES_END
+#	if ZEN_COMPILER_MSC
+#		define ZEN_THIRD_PARTY_INCLUDES_END __pragma(warning(pop))
+#	elif ZEN_COMPILER_CLANG
+#		define ZEN_THIRD_PARTY_INCLUDES_END _Pragma("clang diagnostic pop")
+#	elif ZEN_COMPILER_GCC
+#		define ZEN_THIRD_PARTY_INCLUDES_END _Pragma("GCC diagnostic pop")
+#	endif
+#endif
+
+#if ZEN_COMPILER_MSC
+#	define ZEN_DEBUG_BREAK() \
+		do                    \
+		{                     \
+			__debugbreak();   \
+		} while (0)
+#else
+#	define ZEN_DEBUG_BREAK() \
+		do                    \
+		{                     \
+			__builtin_trap(); \
+		} while (0)
+#endif
+
+//////////////////////////////////////////////////////////////////////////
+// C++20 support
+//
+
+// Clang
+#if ZEN_COMPILER_CLANG && __clang_major__ < 12
+#	error clang-12 onwards is required for C++20 support
+#endif
+
+// GCC
+#if ZEN_COMPILER_GCC && __GNUC__ < 11
+#	error GCC-11 onwards is required for C++20 support
+#endif
+
+// GNU libstdc++
+#if defined(_GLIBCXX_RELEASE) && _GLIBCXX_RELEASE < 11
+#	error GNU libstdc++-11 onwards is required for C++20 support
+#endif
+
+// LLVM libc++
+#if defined(_LIBCPP_VERSION) && _LIBCPP_VERSION < 12000
+#	error LLVM libc++-12 onwards is required for C++20 support
+#endif
+
+//////////////////////////////////////////////////////////////////////////
+// Architecture
+//
+
+#if defined(__amd64__) || defined(_M_X64)
+#	define ZEN_ARCH_X64   1
+#	define ZEN_ARCH_ARM64 0
+#elif defined(__arm64__) || defined(_M_ARM64)
+#	define ZEN_ARCH_X64   0
+#	define ZEN_ARCH_ARM64 1
+#else
+#	error Unknown architecture
+#endif
+
+//////////////////////////////////////////////////////////////////////////
+// Build flavor
+//
+
+#ifdef NDEBUG
+#	define ZEN_BUILD_DEBUG	  0
+#	define ZEN_BUILD_RELEASE 1
+#	define ZEN_BUILD_NAME	  "release"
+#else
+#	define ZEN_BUILD_DEBUG	  1
+#	define ZEN_BUILD_RELEASE 0
+#	define ZEN_BUILD_NAME	  "debug"
+#endif
+
+//////////////////////////////////////////////////////////////////////////
+
+#define ZEN_PLATFORM_SUPPORTS_UNALIGNED_LOADS 1
+
+#if defined(__SIZEOF_WCHAR_T__) && __SIZEOF_WCHAR_T__ == 4
+#	define ZEN_SIZEOF_WCHAR_T 4
+#else
+static_assert(sizeof(wchar_t) == 2, "wchar_t is expected to be two bytes in size");
+#	define ZEN_SIZEOF_WCHAR_T 2
+#endif
+
+//////////////////////////////////////////////////////////////////////////
+
+#ifdef __clang__
+template<typename T>
+auto ZenArrayCountHelper(T& t) -> typename std::enable_if<__is_array(T), char (&)[sizeof(t) / sizeof(t[0]) + 1]>::type;
+#else
+template<typename T, uint32_t N>
+char (&ZenArrayCountHelper(const T (&)[N]))[N + 1];
+#endif
+
+#define ZEN_ARRAY_COUNT(array) (sizeof(ZenArrayCountHelper(array)) - 1)
+
+//////////////////////////////////////////////////////////////////////////
+
+#if ZEN_COMPILER_MSC
+#	define ZEN_NOINLINE	__declspec(noinline)
+#	define ZEN_FORCEINLINE [[msvc::forceinline]]
+#else
+#	define ZEN_NOINLINE	__attribute__((noinline))
+#	define ZEN_FORCEINLINE __attribute__((always_inline))
+#endif
+
+#if ZEN_PLATFORM_WINDOWS
+#	define ZEN_EXE_SUFFIX_LITERAL ".exe"
+#else
+#	define ZEN_EXE_SUFFIX_LITERAL ""
+#endif
+
+#define ZEN_UNUSED(...) ((void)__VA_ARGS__)
+
+//////////////////////////////////////////////////////////////////////////
+
+#if ZEN_COMPILER_MSC
+#	define ZEN_DISABLE_OPTIMIZATION_ACTUAL __pragma(optimize("", off))
+#	define ZEN_ENABLE_OPTIMIZATION_ACTUAL	__pragma(optimize("", on))
+#elif ZEN_COMPILER_GCC
+#	define ZEN_DISABLE_OPTIMIZATION_ACTUAL _Pragma("GCC push_options") _Pragma("GCC optimize (\"O0\")")
+#	define ZEN_ENABLE_OPTIMIZATION_ACTUAL	_Pragma("GCC pop_options")
+#elif ZEN_COMPILER_CLANG
+#	define ZEN_DISABLE_OPTIMIZATION_ACTUAL _Pragma("clang optimize off")
+#	define ZEN_ENABLE_OPTIMIZATION_ACTUAL	_Pragma("clang optimize on")
+#endif
+
+// Set up optimization control macros, now that we have both the build settings and the platform macros
+#define ZEN_DISABLE_OPTIMIZATION ZEN_DISABLE_OPTIMIZATION_ACTUAL
+
+#if ZEN_BUILD_DEBUG
+#	define ZEN_ENABLE_OPTIMIZATION ZEN_DISABLE_OPTIMIZATION_ACTUAL
+#else
+#	define ZEN_ENABLE_OPTIMIZATION ZEN_ENABLE_OPTIMIZATION_ACTUAL
+#endif
+
+#define ZEN_ENABLE_OPTIMIZATION_ALWAYS ZEN_ENABLE_OPTIMIZATION_ACTUAL
+
+#if ZEN_PLATFORM_WINDOWS
+// Tells the compiler to put the decorated function in a certain section (aka. segment) of the executable.
+#	define ZEN_CODE_SECTION(Name) __declspec(code_seg(Name))
+#	define ZEN_DATA_SECTION(Name) __declspec(allocate(Name))
+#	define ZEN_FORCENOINLINE	   __declspec(noinline) /* Force code to NOT be inline */
+#	define LINE_TERMINATOR_ANSI   "\r\n"
+#else
+#	define ZEN_CODE_SECTION(Name)
+#	define ZEN_DATA_SECTION(Name)
+#	define ZEN_FORCENOINLINE
+#	define LINE_TERMINATOR_ANSI "\n"
+#endif
