Add p4ignore file to ZenPluggableTransport

diff --git a/Engine/Source/Developer/ZenPluggableTransport/.p4ignore.txt b/Engine/Source/Developer/ZenPluggableTransport/.p4ignore.txt
new file mode 100644
--- /dev/null
+++ b/Engine/Source/Developer/ZenPluggableTransport/.p4ignore.txt
@@ -0,0 +1,6 @@
+.sarif/
+.vs/
+vcpkg_installed/
+build/
+.xmake/
+vsxmake*/
