Update zenserver transport-sdk to v5.6.6 [f090a03]
#rnx

diff --git a/Engine/Source/Developer/ZenPluggableTransport/transport-sdk/include/transportplugin.h b/Engine/Source/Developer/ZenPluggableTransport/transport-sdk/include/transportplugin.h
--- a/Engine/Source/Developer/ZenPluggableTransport/transport-sdk/include/transportplugin.h
+++ b/Engine/Source/Developer/ZenPluggableTransport/transport-sdk/include/transportplugin.h
@@ -17,10 +17,14 @@
 
 namespace zen {
 
+// Current API version, value will be incremented to represent breaking changes
+static const uint32_t kTransportApiVersion = 1;
+
 class TransportConnection;
 class TransportPlugin;
 class TransportServerConnection;
 class TransportServer;
+class TransportLogger;
 
 /*************************************************************************
 
@@ -60,6 +64,29 @@ public:
 	virtual TransportServerConnection* CreateConnectionHandler(TransportConnection* Connection) = 0;
 };
 
+/** Logger interface
+
+ There will be one instance of this provided by the system to the transport plugin
+
+ The plugin can use this to log messages back to zen server
+
+ */
+class TransportLogger
+{
+public:
+	enum class LogLevel : uint32_t
+	{
+		Trace	 = 0,
+		Debug	 = 1,
+		Info	 = 2,
+		Warn	 = 3,
+		Err		 = 4,
+		Critical = 5,
+	};
+
+	virtual void LogMessage(LogLevel Level, const char* Message) = 0;
+};
+
 /*************************************************************************
 
    The following interfaces are to be implemented by transport plugins.
@@ -116,7 +143,18 @@ public:
 
 extern "C"
 {
-	DLL_TRANSPORT_API zen::TransportPlugin* CreateTransportPlugin();
+	/** Provide information about plugin version
+
+	   Fills out API version (kTransportApiVersion) plugin was built against.
+	   Fills out plugin own version ever increasing version number,
+	   a copy of plugin with higher version will be used.
+	*/
+	DLL_TRANSPORT_API void GetTransportPluginVersion(uint32_t* OutApiVersion, uint32_t* OutPluginVersion);
+
+	// Return nullptr if requested api version mismatches api version plugin was built against
+	DLL_TRANSPORT_API zen::TransportPlugin* CreateTransportPlugin(zen::TransportLogger* Logger);
 }
 
-typedef zen::TransportPlugin* (*PfnCreateTransportPlugin)();
+typedef void (*PfnGetTransportPluginVersion)(uint32_t* OutApiVersion, uint32_t* OutPluginVersion);
+
+typedef zen::TransportPlugin* (*PfnCreateTransportPlugin)(zen::TransportLogger* Logger);
diff --git a/Engine/Source/Developer/ZenPluggableTransport/winsock/winsock.cpp b/Engine/Source/Developer/ZenPluggableTransport/winsock/winsock.cpp
--- a/Engine/Source/Developer/ZenPluggableTransport/winsock/winsock.cpp
+++ b/Engine/Source/Developer/ZenPluggableTransport/winsock/winsock.cpp
@@ -317,7 +317,7 @@ WinsockTransportPlugin::Initialize(TransportServer* ServerInterface)
 					{
 						Connection->HandleConnection();
 					}
-					catch (std::exception&)
+					catch (const std::exception&)
 					{
 						// ZEN_WARN("exception caught in connection loop: {}", Ex.what());
 					}
@@ -364,8 +364,22 @@ WinsockTransportPlugin::IsAvailable()
 
 //////////////////////////////////////////////////////////////////////////
 
+void
+GetTransportPluginVersion(uint32_t* OutApiVersion, uint32_t* OutPluginVersion)
+{
+	if (OutApiVersion != nullptr)
+	{
+		*OutApiVersion = kTransportApiVersion;
+	}
+
+	if (OutPluginVersion != nullptr)
+	{
+		*OutPluginVersion = 1;
+	}
+}
+
 TransportPlugin*
-CreateTransportPlugin()
+CreateTransportPlugin([[maybe_unused]] TransportLogger* Logger)
 {
 	return new WinsockTransportPlugin;
 }
